Add Mass Processing Phase event dependency option in UAF. This is contained in a new plugin: UAFMass

diff --git a/Engine/Plugins/Experimental/UAF/UAFMass/Source/UAFMass/Private/Module/RigVMTrait_ModuleEventDependency_MassPhaseProcessor.cpp b/Engine/Plugins/Experimental/UAF/UAFMass/Source/UAFMass/Private/Module/RigVMTrait_ModuleEventDependency_MassPhaseProcessor.cpp
new file mode 100644
--- /dev/null
+++ b/Engine/Plugins/Experimental/UAF/UAFMass/Source/UAFMass/Private/Module/RigVMTrait_ModuleEventDependency_MassPhaseProcessor.cpp
@@ -0,0 +1,57 @@
+// Copyright Epic Games, Inc. All Rights Reserved.
+
+#include "RigVMTrait_ModuleEventDependency_MassPhaseProcessor.h"
+
+#include "MassSimulationSubsystem.h"
+
+
+#if WITH_EDITOR
+FString FRigVMTrait_ModuleEventDependency_MassPhaseProcessor::GetDisplayName() const
+{
+	return StaticStruct()->GetDisplayNameText().ToString();
+}
+#endif // WITH_EDITOR
+
+void FRigVMTrait_ModuleEventDependency_MassPhaseProcessor::OnAddDependency(const UE::UAF::FModuleDependencyContext& InContext) const
+{
+	check(IsInGameThread()); // TickFunction add / remove prerequisite is not thread safe
+	UMassSimulationSubsystem* MassSimulationSubsystem = InContext.Object->GetWorld()->GetSubsystem<UMassSimulationSubsystem>();
+
+	if (MassSimulationSubsystem == nullptr)
+	{
+		return;
+	}
+
+	FMassProcessingPhaseManager& MassSimulationPhaseManager = MassSimulationSubsystem->GetMutablePhaseManager(); 
+	FTickFunction& ProcessingPhaseTickFunction = MassSimulationPhaseManager.GetProcessingPhaseTickFunction(DependentMassPhase);
+	
+	if (Ordering == EAnimNextModuleEventDependencyOrdering::Before)
+	{
+		InContext.TickFunction.AddPrerequisite(MassSimulationSubsystem, ProcessingPhaseTickFunction);
+	}
+	else
+	{
+		UActorComponent* AnimNextComponent = Cast<UActorComponent>(InContext.Object);
+		ProcessingPhaseTickFunction.AddPrerequisite(AnimNextComponent, InContext.TickFunction);
+	}
+}
+
+
+void FRigVMTrait_ModuleEventDependency_MassPhaseProcessor::OnRemoveDependency(const UE::UAF::FModuleDependencyContext& InContext) const
+{
+	check(IsInGameThread()); // TickFunction add / remove prerequisite is not thread safe
+	UMassSimulationSubsystem* MassSimulationSubsystem = InContext.Object->GetWorld()->GetSubsystem<UMassSimulationSubsystem>();
+
+	FMassProcessingPhaseManager& MassSimulationPhaseManager = MassSimulationSubsystem->GetMutablePhaseManager(); 
+	FTickFunction& ProcessingPhaseTickFunction = MassSimulationPhaseManager.GetProcessingPhaseTickFunction(DependentMassPhase);
+	
+	if (Ordering == EAnimNextModuleEventDependencyOrdering::Before)
+	{
+		InContext.TickFunction.RemovePrerequisite(MassSimulationSubsystem, ProcessingPhaseTickFunction);
+	}
+	else
+	{
+		UActorComponent* AnimNextComponent = Cast<UActorComponent>(InContext.Object);
+		ProcessingPhaseTickFunction.RemovePrerequisite(AnimNextComponent, InContext.TickFunction);
+	}
+}
\ No newline at end of file
diff --git a/Engine/Plugins/Experimental/UAF/UAFMass/Source/UAFMass/Private/Module/RigVMTrait_ModuleEventDependency_MassPhaseProcessor.h b/Engine/Plugins/Experimental/UAF/UAFMass/Source/UAFMass/Private/Module/RigVMTrait_ModuleEventDependency_MassPhaseProcessor.h
new file mode 100644
--- /dev/null
+++ b/Engine/Plugins/Experimental/UAF/UAFMass/Source/UAFMass/Private/Module/RigVMTrait_ModuleEventDependency_MassPhaseProcessor.h
@@ -0,0 +1,30 @@
+// Copyright Epic Games, Inc. All Rights Reserved.
+
+#pragma once
+
+#include "MassProcessingTypes.h"
+#include "Module/RigVMTrait_ModuleEventDependency.h"
+#include "RigVMTrait_ModuleEventDependency_MassPhaseProcessor.generated.h"
+
+
+/** A dependency on one of a mass phase processors (equates to a tick function)
+ */
+USTRUCT(DisplayName = "Mass Phase Processor", meta = (ShowTooltip = true))
+struct FRigVMTrait_ModuleEventDependency_MassPhaseProcessor : public FRigVMTrait_ModuleEventDependency
+{
+	GENERATED_BODY()
+
+	// FRigVMTrait interface
+#if WITH_EDITOR
+	virtual FString GetDisplayName() const override;
+#endif
+
+	// FAnimNextModuleEventDependency interface
+	virtual void OnAddDependency(const UE::UAF::FModuleDependencyContext& InContext) const override;
+	virtual void OnRemoveDependency(const UE::UAF::FModuleDependencyContext& InContext) const override;
+
+
+	// The Mass processing phase that this dependency relates to
+	UPROPERTY(EditAnywhere, Category = "Dependency")
+	EMassProcessingPhase DependentMassPhase = EMassProcessingPhase::PrePhysics;
+};
\ No newline at end of file
diff --git a/Engine/Plugins/Experimental/UAF/UAFMass/Source/UAFMass/Private/UAFMassModule.cpp b/Engine/Plugins/Experimental/UAF/UAFMass/Source/UAFMass/Private/UAFMassModule.cpp
new file mode 100644
--- /dev/null
+++ b/Engine/Plugins/Experimental/UAF/UAFMass/Source/UAFMass/Private/UAFMassModule.cpp
@@ -0,0 +1,21 @@
+// Copyright Epic Games, Inc. All Rights Reserved.
+#include "UAFMassModule.h"
+
+#include "CoreMinimal.h"
+#include "RigVMCore/RigVMRegistry.h"
+
+#include "Modules/ModuleManager.h"
+
+
+void FUAFMassModule::StartupModule()
+{
+
+}
+
+void FUAFMassModule::ShutdownModule()
+{
+	// This function may be called during shutdown to clean up your module.  For modules that support dynamic reloading,
+	// we call this function before unloading the module.
+}
+
+IMPLEMENT_MODULE(FUAFMassModule, UAFMass)
\ No newline at end of file
diff --git a/Engine/Plugins/Experimental/UAF/UAFMass/Source/UAFMass/Public/UAFMassModule.h b/Engine/Plugins/Experimental/UAF/UAFMass/Source/UAFMass/Public/UAFMassModule.h
new file mode 100644
--- /dev/null
+++ b/Engine/Plugins/Experimental/UAF/UAFMass/Source/UAFMass/Public/UAFMassModule.h
@@ -0,0 +1,15 @@
+// Copyright Epic Games, Inc. All Rights Reserved.
+
+#pragma once
+
+#include "CoreMinimal.h"
+#include "Modules/ModuleManager.h"
+
+class FUAFMassModule : public IModuleInterface
+{
+public:
+
+	/** IModuleInterface implementation */
+	virtual void StartupModule() override;
+	virtual void ShutdownModule() override;
+};
diff --git a/Engine/Plugins/Experimental/UAF/UAFMass/Source/UAFMass/UAFMass.Build.cs b/Engine/Plugins/Experimental/UAF/UAFMass/Source/UAFMass/UAFMass.Build.cs
new file mode 100644
--- /dev/null
+++ b/Engine/Plugins/Experimental/UAF/UAFMass/Source/UAFMass/UAFMass.Build.cs
@@ -0,0 +1,32 @@
+// Copyright Epic Games, Inc. All Rights Reserved.
+
+using UnrealBuildTool;
+
+public class UAFMass : ModuleRules
+{
+	public UAFMass(ReadOnlyTargetRules Target) : base(Target)
+	{
+
+		PublicDependencyModuleNames.AddRange(
+			new string[]
+			{
+				"Core",
+				"UAF"
+			}
+			);
+			
+		
+		PrivateDependencyModuleNames.AddRange(
+			new string[]
+			{
+				"CoreUObject",
+				"RigVM",
+				"Engine",
+				"MassSimulation",
+				"MassEntity"
+			}
+			);
+
+		SetupGameplayDebuggerSupport(Target);
+	}
+}
diff --git a/Engine/Plugins/Experimental/UAF/UAFMass/UAFMass.uplugin b/Engine/Plugins/Experimental/UAF/UAFMass/UAFMass.uplugin
new file mode 100644
--- /dev/null
+++ b/Engine/Plugins/Experimental/UAF/UAFMass/UAFMass.uplugin
@@ -0,0 +1,41 @@
+{
+	"FileVersion" : 3,
+	"Version" : 1,
+	"VersionName" : "0.1",
+	"FriendlyName" : "UAF Mass",
+	"Description" : "Mass integration for UAF.",
+	"Category" : "Animation",
+	"CreatedBy" : "Epic Games, Inc.",
+	"CreatedByURL" : "https://epicgames.com",
+	"DocsURL" : "",
+	"MarketplaceURL" : "",
+	"SupportURL" : "",
+	"EnabledByDefault" : false,
+	"CanContainContent" : false,
+	"IsBetaVersion": false,
+	"IsExperimentalVersion": true,
+	"Installed" : false,
+	"Plugins":
+	[
+		{
+			"Name": "MassGameplay",
+			"Enabled": true
+		},
+		{
+			"Name": "UAF",
+			"Enabled": true
+		},
+		{
+			"Name": "UAFAnimGraph",
+			"Enabled": true
+		},
+	],
+	"Modules" :
+	[
+		{
+			"Name" : "UAFMass",
+			"Type" : "Runtime",
+			"LoadingPhase" : "Default"
+		}
+	]
+}
