Allow running UAF update from within a Mass Processor
#Synced-CL 48062098

diff --git a/Engine/Plugins/Experimental/UAF/UAF/Source/UAF/Private/Module/SystemReference.cpp b/Engine/Plugins/Experimental/UAF/UAF/Source/UAF/Private/Module/SystemReference.cpp
--- a/Engine/Plugins/Experimental/UAF/UAF/Source/UAF/Private/Module/SystemReference.cpp
+++ b/Engine/Plugins/Experimental/UAF/UAF/Source/UAF/Private/Module/SystemReference.cpp
@@ -116,7 +116,7 @@ FSystemReference::FSystemReference(FAnimNextModuleInstance& InInstance)
 
 FSystemReference::FSystemReference(const UObject* InAsset, UObject* InObject, EAnimNextModuleInitMethod InitMethod)
 {
-	check(IsInGameThread());
+	check(IsInGameThread() || IsInParallelGameThread());
 	{
 		LLM_SCOPE_BYNAME(TEXT("UAF"));
 		
@@ -254,7 +254,7 @@ void FSystemReference::QueueInputTraitEvent(FAnimNextTraitEventPtr Event) const
 const FTickFunction* FSystemReference::FindTickFunction(FName InEventName) const
 {
 	using namespace UE::UAF;
-	check(IsInGameThread());
+	check(IsInGameThread() || IsInParallelGameThread());
 	if (!IsValid())
 	{
 		UE_LOGFMT(LogAnimation, Warning, "FindTickFunction: Invalid system handle");
@@ -416,8 +416,28 @@ EPropertyBagResult FSystemReference::WriteVariable(const FAnimNextVariableRefere
 	return EPropertyBagResult::PropertyNotFound;
 }
 
+void FSystemReference::RunEvent(FName InEventName, float InDeltaTime) const
+{
+	using namespace UE::UAF;
+	check(IsInGameThread() || IsInParallelGameThread());
+	if (!IsValid())
+	{
+		UE_LOGFMT(LogAnimation, Warning, "RunEvent: Invalid system handle");
+		return;
+	}
+
+	FModuleEventTickFunction* TickFunction = Ptr->FindTickFunctionByName(InEventName);
+	if (TickFunction == nullptr)
+	{
+		UE_LOGFMT(LogAnimation, Warning, "RunEvent: Could not find event '{EventName}' in module '{ModuleName}'", InEventName, Ptr->GetAssetName());
+		return;
+	}
+
+	TickFunction->Run(InDeltaTime);
 }
 
+} //end namespace UE::UAF
+
 UUAFEngineSubsystem::UUAFEngineSubsystem()
 {
 	if (!HasAnyFlags(RF_ClassDefaultObject))
diff --git a/Engine/Plugins/Experimental/UAF/UAF/Source/UAF/Public/AnimNextRigVMAsset.h b/Engine/Plugins/Experimental/UAF/UAF/Source/UAF/Public/AnimNextRigVMAsset.h
--- a/Engine/Plugins/Experimental/UAF/UAF/Source/UAF/Public/AnimNextRigVMAsset.h
+++ b/Engine/Plugins/Experimental/UAF/UAF/Source/UAF/Public/AnimNextRigVMAsset.h
@@ -34,6 +34,11 @@ namespace UE::UAF
 	struct FInjectionInfo;
 }
 
+namespace UE::UAF::Tests
+{
+	class FSystemReference_Creation;
+}
+
 namespace UE::UAF::UncookedOnly
 {
 	struct FUtils;
@@ -70,6 +75,7 @@ protected:
 	friend UE::UAF::FAnimGraphBuilderContext;
 	friend UE::UAF::FSystemBuilderContext;
 	friend UE::UAF::UncookedOnly::FCompilationScope;
+	friend UE::UAF::Tests::FSystemReference_Creation;
 	friend FUAFAssetInstance;
 
 	UE_API UAnimNextRigVMAsset(const FObjectInitializer& ObjectInitializer);
diff --git a/Engine/Plugins/Experimental/UAF/UAF/Source/UAF/Public/Module/SystemReference.h b/Engine/Plugins/Experimental/UAF/UAF/Source/UAF/Public/Module/SystemReference.h
--- a/Engine/Plugins/Experimental/UAF/UAF/Source/UAF/Public/Module/SystemReference.h
+++ b/Engine/Plugins/Experimental/UAF/UAF/Source/UAF/Public/Module/SystemReference.h
@@ -1,4 +1,4 @@
-﻿// Copyright Epic Games, Inc. All Rights Reserved.
+// Copyright Epic Games, Inc. All Rights Reserved.
 
 #pragma once
 
@@ -20,6 +20,9 @@ struct FTickFunction;
 class UAnimNextComponent;
 struct FUAFWeakSystemReference;
 class UAnimNextRigVMAsset;
+struct FMassUAFFragment;
+class UMassUAFProcessor;
+class UMassUAFInitializer;
 
 namespace UE::UAF
 {
@@ -30,6 +33,11 @@ struct FModuleEventTickFunction;
 class FAnimNextModuleImpl;
 }
 
+namespace UE::UAF::Tests
+{
+class FSystemReference_Creation;
+}
+
 namespace UE::UAF
 {
 
@@ -39,19 +47,25 @@ namespace UE::UAF
 struct FSystemReference
 {
 	// Default-constructable only because it needs to be embedded in UObjects/UStructs
-	FSystemReference() = default;
+	UE_API FSystemReference() = default;
 
 	// Move only, non-copyable
 	FSystemReference(const FSystemReference&) = delete;
 	FSystemReference& operator=(const FSystemReference&) = delete;
-	FSystemReference(FSystemReference&&);
-	FSystemReference& operator=(FSystemReference&&);
+	UE_API FSystemReference(FSystemReference&&);
+	UE_API FSystemReference& operator=(FSystemReference&&);
 	
 private:
 	friend UAnimNextComponent;
 	friend FUAFWeakSystemReference;
 	friend UE::UAF::FModuleEventTickFunction;
 	friend UE::UAF::FAnimNextModuleImpl;
+	friend UE::UAF::Tests::FSystemReference_Creation;
+
+	// @TODO: Temporarily friending Mass types, so that Mass can act as a system host similar to UAnimNextComponent.
+	friend FMassUAFFragment;
+	friend UMassUAFProcessor;
+	friend UMassUAFInitializer;
 
 	// Construct an owning handle to a system instance.
 	// Constructs a FAnimNextModuleInstance for use by worker threads.
@@ -105,6 +119,11 @@ private:
 	// @param	InEventName			The event associated to the wanted tick function
 	UE_API const FTickFunction* FindTickFunction(FName InEventName) const;
 
+	// Run a specific Event in-place.
+	// @TODO: Should only be called on unregistered tick functions.
+	// @param	InEventName		The event we want to run.
+	UE_API void RunEvent(FName InEventName, float DeltaTime) const;
+
 	// Add a dependency on a tick function to the specified event
 	// Should be called only from the game thread
 	// @param	InObject			The object that owns the tick function
diff --git a/Engine/Plugins/Experimental/UAF/UAFMass/Source/UAFMass/Private/Fragments/MassUAFFragment.h b/Engine/Plugins/Experimental/UAF/UAFMass/Source/UAFMass/Private/Fragments/MassUAFFragment.h
new file mode 100644
--- /dev/null
+++ b/Engine/Plugins/Experimental/UAF/UAFMass/Source/UAFMass/Private/Fragments/MassUAFFragment.h
@@ -0,0 +1,40 @@
+// Copyright Epic Games, Inc. All Rights Reserved.
+
+#pragma once
+
+#include "MassCommonFragments.h"
+#include "Module/SystemReference.h"
+#include "MassUAFFragment.generated.h"
+
+
+/**
+ * Fragment owning the UAF system associated to the entity.
+ */
+USTRUCT()
+struct FMassUAFFragment : public FObjectWrapperFragment
+{
+	GENERATED_BODY()
+	// @TODO: We don't need this data after UAF system creation. Remove this when we figure out the correct init flow i.e. shared/chunk fragment or Subsystem ownership.
+	UPROPERTY(transient)
+	TObjectPtr<const UObject> Asset = nullptr;
+
+	UE::UAF::FSystemReference SystemReference;
+};
+
+template<> 
+struct TStructOpsTypeTraits<FMassUAFFragment> : public TStructOpsTypeTraitsBase2<FMassUAFFragment> 
+{ 
+	enum 
+	{ 
+		WithCopy = false 
+	}; 
+};
+
+template<>
+struct TMassFragmentTraits<FMassUAFFragment> final
+{
+	enum
+	{
+		AuthorAcceptsItsNotTriviallyCopyable = true
+	};
+};
diff --git a/Engine/Plugins/Experimental/UAF/UAFMass/Source/UAFMass/Private/Processors/MassUAFProcessor.cpp b/Engine/Plugins/Experimental/UAF/UAFMass/Source/UAFMass/Private/Processors/MassUAFProcessor.cpp
new file mode 100644
--- /dev/null
+++ b/Engine/Plugins/Experimental/UAF/UAFMass/Source/UAFMass/Private/Processors/MassUAFProcessor.cpp
@@ -0,0 +1,141 @@
+// Copyright Epic Games, Inc. All Rights Reserved.
+
+#include "MassUAFProcessor.h"
+
+#include "Fragments/MassUAFFragment.h"
+#include "MassExecutionContext.h"
+#include "MassCommonTypes.h"
+#include "Module/ModuleTickFunction.h"
+#include "Module/AnimNextModuleInstance.h"
+#include "Module/RigUnit_AnimNextModuleEvents.h"
+#include "Subsystems/MassUAFSubsystem.h"
+
+DEFINE_LOG_CATEGORY(LogMassUAF);
+
+namespace UE::Mass::ProcessorGroupNames
+{
+	static const FName Animation = FName(TEXT("Animation"));
+	static const FName AnimationDebug = FName(TEXT("AnimationDebug"));
+} // namespace UE::Mass::ProcessorGroupNames
+
+namespace UE::UAF::Mass
+{
+	static bool bDebug = false;
+	FAutoConsoleVariableRef CVarDebugAnimPose(TEXT("Mass.UAF.Debug"), bDebug, TEXT(""));
+}
+
+UMassUAFInitializer::UMassUAFInitializer()
+	: EntityQuery(*this)
+{
+	ObservedType = FMassUAFFragment::StaticStruct();
+	ObservedOperations = EMassObservedOperationFlags::Add;
+	ExecutionFlags = (int32)EProcessorExecutionFlags::AllNetModes;
+	bRequiresGameThreadExecution = false;
+}
+
+void UMassUAFInitializer::ConfigureQueries(const TSharedRef<FMassEntityManager>& EntityManager)
+{
+	EntityQuery.AddRequirement<FMassUAFFragment>(EMassFragmentAccess::ReadWrite);
+	EntityQuery.AddSubsystemRequirement<UMassUAFSubsystem>(EMassFragmentAccess::ReadWrite);
+}
+
+void UMassUAFInitializer::Execute(FMassEntityManager& EntityManager, FMassExecutionContext& Context)
+{
+	using namespace UE::UAF;	
+	EntityQuery.ForEachEntityChunk(Context, [this](FMassExecutionContext& Context)
+	{
+		FTaskTagScope Scope(ETaskTag::EParallelGameThread);
+
+		UMassUAFSubsystem& UAFSubsystem = Context.GetMutableSubsystemChecked<UMassUAFSubsystem>();
+		const int32 NumEntities = Context.GetNumEntities();
+		TArrayView<FMassUAFFragment> UAFFragmentList = Context.GetMutableFragmentView<FMassUAFFragment>();
+
+		for (FMassExecutionContext::FEntityIterator EntityIt = Context.CreateEntityIterator(); EntityIt; ++EntityIt)
+		{
+			FMassUAFFragment& UAFFragment = UAFFragmentList[EntityIt];
+
+			if (UAFFragment.Asset)
+			{
+				// @TODO: Passing a dummy UAFSubsystem since it works well with rewind debugger + we need a UObject owner, but should go away once we agree on system ownership.
+				UAFFragment.SystemReference = FSystemReference(UAFFragment.Asset, &UAFSubsystem, EAnimNextModuleInitMethod::None);
+			}
+		}
+	});
+}
+
+UMassUAFProcessor::UMassUAFProcessor() : EntityQuery(*this)
+{
+	ExecutionFlags = (int32)EProcessorExecutionFlags::AllNetModes;
+	ExecutionOrder.ExecuteInGroup = UE::Mass::ProcessorGroupNames::Animation;
+	ExecutionOrder.ExecuteAfter.Add(UE::Mass::ProcessorGroupNames::Movement);
+	bRequiresGameThreadExecution = false;
+}
+
+void UMassUAFProcessor::ConfigureQueries(const TSharedRef<FMassEntityManager>& EntityManager)
+{
+	EntityQuery.AddRequirement<FMassUAFFragment>(EMassFragmentAccess::ReadOnly);
+}
+
+void UMassUAFProcessor::Execute(FMassEntityManager& EntityManager, FMassExecutionContext& Context)
+{
+	using namespace UE::UAF;	
+	EntityQuery.ParallelForEachEntityChunk(Context, [this](FMassExecutionContext& Context)
+	{
+		FTaskTagScope Scope(ETaskTag::EParallelGameThread);
+		const int32 NumEntities = Context.GetNumEntities();
+		const TConstArrayView<FMassUAFFragment> UAFFragmentList = Context.GetFragmentView<FMassUAFFragment>();
+		
+		for (FMassExecutionContext::FEntityIterator EntityIt = Context.CreateEntityIterator(); EntityIt; ++EntityIt)
+		{
+			FMassEntityHandle Entity = Context.GetEntity(EntityIt);
+			const FMassUAFFragment& UAFFragment = UAFFragmentList[EntityIt];
+
+			if (UAFFragment.SystemReference.IsValid())
+			{
+				// @TODO: Pass Mass execution context to UAF update
+				//SystemReference->QueueTask(NAME_None, [&Context, Entity, &UAFFragment](const UE::UAF::FModuleTaskContext& InContext)
+				//{
+				//How do we do this? There's no ExecuteContext accessible here.
+				//});
+
+				UAFFragment.SystemReference.RunEvent(FRigUnit_AnimNextPrePhysicsEvent::DefaultEventName, Context.GetDeltaTimeSeconds());
+			}
+		}
+	});
+}
+
+UMassAnimPoseDebugProcessor::UMassAnimPoseDebugProcessor() : EntityQuery(*this)
+{
+	ExecutionFlags = (int32)EProcessorExecutionFlags::AllNetModes;
+	ExecutionOrder.ExecuteInGroup = UE::Mass::ProcessorGroupNames::AnimationDebug;
+	ExecutionOrder.ExecuteAfter.Add(UE::Mass::ProcessorGroupNames::Animation);
+}
+
+void UMassAnimPoseDebugProcessor::ConfigureQueries(const TSharedRef<FMassEntityManager>& EntityManager)
+{
+	EntityQuery.AddRequirement<FTransformFragment>(EMassFragmentAccess::ReadOnly);
+	EntityQuery.AddSubsystemRequirement<UMassUAFSubsystem>(EMassFragmentAccess::ReadWrite);
+}
+
+void UMassAnimPoseDebugProcessor::Execute(FMassEntityManager& EntityManager, FMassExecutionContext& Context)
+{
+	UWorld* World = GetWorld();
+	if (World == nullptr || (UE::UAF::Mass::bDebug == false))
+	{
+		return;
+	}
+
+	// @TODO: Update this when Mass has a way to access UAF's pose buffer.
+	/*using namespace UE::UAF;	
+	EntityQuery.ParallelForEachEntityChunk(Context, [this, World](FMassExecutionContext& Context)
+	{
+		UMassUAFSubsystem& UAFSubsystem = Context.GetMutableSubsystemChecked<UMassUAFSubsystem>();
+
+		const int32 NumEntities = Context.GetNumEntities();
+		const TConstArrayView<FTransformFragment> TransformList = Context.GetFragmentView<FTransformFragment>();
+		
+		for (FMassExecutionContext::FEntityIterator EntityIt = Context.CreateEntityIterator(); EntityIt; ++EntityIt)
+		{
+		}
+	});*/
+}
\ No newline at end of file
diff --git a/Engine/Plugins/Experimental/UAF/UAFMass/Source/UAFMass/Private/Processors/MassUAFProcessor.h b/Engine/Plugins/Experimental/UAF/UAFMass/Source/UAFMass/Private/Processors/MassUAFProcessor.h
new file mode 100644
--- /dev/null
+++ b/Engine/Plugins/Experimental/UAF/UAFMass/Source/UAFMass/Private/Processors/MassUAFProcessor.h
@@ -0,0 +1,57 @@
+// Copyright Epic Games, Inc. All Rights Reserved.
+
+#pragma once
+
+#include "MassEntityQuery.h"
+#include "MassProcessor.h"
+#include "MassObserverProcessor.h"
+
+#include "MassUAFProcessor.generated.h"
+
+DECLARE_LOG_CATEGORY_EXTERN(LogMassUAF, Display, All);
+
+UCLASS()
+class UMassUAFInitializer : public UMassObserverProcessor
+{
+	GENERATED_BODY()
+
+public:
+	UMassUAFInitializer();
+
+protected:
+	virtual void ConfigureQueries(const TSharedRef<FMassEntityManager>& EntityManager) override;
+	virtual void Execute(FMassEntityManager& EntityManager, FMassExecutionContext& Context) override;
+
+	FMassEntityQuery EntityQuery;
+};
+
+
+UCLASS()
+class UMassUAFProcessor : public UMassProcessor
+{
+	GENERATED_BODY()
+
+public:
+	UMassUAFProcessor();
+
+protected:
+	virtual void ConfigureQueries(const TSharedRef<FMassEntityManager>& EntityManager) override;
+	virtual void Execute(FMassEntityManager& EntityManager, FMassExecutionContext& Context) override;
+
+	FMassEntityQuery EntityQuery;
+};
+
+UCLASS()
+class UMassAnimPoseDebugProcessor : public UMassProcessor
+{
+	GENERATED_BODY()
+
+public:
+	UMassAnimPoseDebugProcessor();
+
+protected:
+	virtual void ConfigureQueries(const TSharedRef<FMassEntityManager>& EntityManager) override;
+	virtual void Execute(FMassEntityManager& EntityManager, FMassExecutionContext& Context) override;
+
+	FMassEntityQuery EntityQuery;
+};
\ No newline at end of file
diff --git a/Engine/Plugins/Experimental/UAF/UAFMass/Source/UAFMass/Private/Subsystems/MassUAFSubsystem.h b/Engine/Plugins/Experimental/UAF/UAFMass/Source/UAFMass/Private/Subsystems/MassUAFSubsystem.h
new file mode 100644
--- /dev/null
+++ b/Engine/Plugins/Experimental/UAF/UAFMass/Source/UAFMass/Private/Subsystems/MassUAFSubsystem.h
@@ -0,0 +1,27 @@
+// Copyright Epic Games, Inc. All Rights Reserved.
+
+#pragma once
+
+#include "MassEntityHandle.h"
+#include "MassEntityTypes.h"
+#include "MassSubsystemBase.h"
+#include "MassUAFSubsystem.generated.h"
+
+/** @TODO: Dummy subsystem used as UObject owner of Mass - owned UAF systems.
+ * Remove this when we have proper ownership of UAF system + per-entity rewind debugger support. 
+ */
+UCLASS(MinimalAPI)
+class UMassUAFSubsystem : public UMassSubsystemBase
+{
+	GENERATED_BODY()	
+};
+
+template<>
+struct TMassExternalSubsystemTraits<UMassUAFSubsystem> final
+{
+	enum
+	{
+		GameThreadOnly = false,
+		ThreadSafeWrite = true
+	};
+};
diff --git a/Engine/Plugins/Experimental/UAF/UAFMass/Source/UAFMass/Private/Traits/MassUAFTrait.cpp b/Engine/Plugins/Experimental/UAF/UAFMass/Source/UAFMass/Private/Traits/MassUAFTrait.cpp
new file mode 100644
--- /dev/null
+++ b/Engine/Plugins/Experimental/UAF/UAFMass/Source/UAFMass/Private/Traits/MassUAFTrait.cpp
@@ -0,0 +1,19 @@
+// Copyright Epic Games, Inc. All Rights Reserved.
+
+#include "MassUAFTrait.h"
+
+#include "Fragments/MassUAFFragment.h"
+#include "MassEntityTemplateRegistry.h"
+#include "MassEntityView.h"
+#include "MassTranslator.h"
+#include "Component/AnimNextComponent.h"
+
+#include UE_INLINE_GENERATED_CPP_BY_NAME(MassUAFTrait)
+
+void UMassUAFTrait::BuildTemplate(FMassEntityTemplateBuildContext& BuildContext, const UWorld& World) const
+{
+	using namespace UE::UAF;
+
+	FMassUAFFragment& UAFFragment = BuildContext.AddFragment_GetRef<FMassUAFFragment>();
+	UAFFragment.Asset = Asset;
+}
\ No newline at end of file
diff --git a/Engine/Plugins/Experimental/UAF/UAFMass/Source/UAFMass/Private/Traits/MassUAFTrait.h b/Engine/Plugins/Experimental/UAF/UAFMass/Source/UAFMass/Private/Traits/MassUAFTrait.h
new file mode 100644
--- /dev/null
+++ b/Engine/Plugins/Experimental/UAF/UAFMass/Source/UAFMass/Private/Traits/MassUAFTrait.h
@@ -0,0 +1,21 @@
+// Copyright Epic Games, Inc. All Rights Reserved.
+
+#pragma once
+
+#include "MassEntityTraitBase.h"
+#include "MassUAFTrait.generated.h"
+
+struct FMassEntityTemplateBuildContext;
+
+UCLASS(MinimalAPI, BlueprintType, EditInlineNew, CollapseCategories, meta = (DisplayName = "Mass UAF Trait"))
+class UMassUAFTrait : public UMassEntityTraitBase
+{
+	GENERATED_BODY()
+
+	/** The asset that this Mass Entity will run in UAF */
+	UPROPERTY(EditAnywhere, Category = "Asset", DisplayName = "Asset", meta = (GetAllowedClasses = "/Script/UAF.AnimNextConfig:GetAllowedAssetClasses"))
+	TObjectPtr<const UObject> Asset = nullptr;
+
+protected:
+	virtual void BuildTemplate(FMassEntityTemplateBuildContext& BuildContext, const UWorld& World) const override;
+};
diff --git a/Engine/Plugins/Experimental/UAF/UAFMass/Source/UAFMass/UAFMass.Build.cs b/Engine/Plugins/Experimental/UAF/UAFMass/Source/UAFMass/UAFMass.Build.cs
--- a/Engine/Plugins/Experimental/UAF/UAFMass/Source/UAFMass/UAFMass.Build.cs
+++ b/Engine/Plugins/Experimental/UAF/UAFMass/Source/UAFMass/UAFMass.Build.cs
@@ -22,6 +22,8 @@ public class UAFMass : ModuleRules
 				"CoreUObject",
 				"RigVM",
 				"Engine",
+				"MassCommon",
+				"MassSpawner",
 				"MassSimulation",
 				"MassEntity"
 			}
diff --git a/Engine/Plugins/Experimental/UAF/UAFTests/Source/UAFTestSuite/Private/AnimNextTest.cpp b/Engine/Plugins/Experimental/UAF/UAFTests/Source/UAFTestSuite/Private/AnimNextTest.cpp
--- a/Engine/Plugins/Experimental/UAF/UAFTests/Source/UAFTestSuite/Private/AnimNextTest.cpp
+++ b/Engine/Plugins/Experimental/UAF/UAFTests/Source/UAFTestSuite/Private/AnimNextTest.cpp
@@ -55,6 +55,23 @@ void FUtils::CleanupAfterTests()
 	CollectGarbage(GARBAGE_COLLECTION_KEEPFLAGS);
 }
 
+class UWorld* FUtils::GetWorld()
+{
+#if WITH_EDITOR
+	if (GIsEditor)
+	{
+		return GWorld;
+	}
+#endif // WITH_EDITOR
+
+	if (GEngine->GetWorldContexts().IsEmpty())
+	{
+		return nullptr;
+	}
+
+	return GEngine->GetWorldContexts()[0].World();
+}
+
 IMPLEMENT_SIMPLE_AUTOMATION_TEST(FDataRegistryTransformsTest, "Animation.AnimNext.AnimationDataRegistry.TransformsTest", EAutomationTestFlags::EditorContext | EAutomationTestFlags::EngineFilter)
 
 bool FDataRegistryTransformsTest::RunTest(const FString& InParameters)
diff --git a/Engine/Plugins/Experimental/UAF/UAFTests/Source/UAFTestSuite/Public/AnimNextTest.h b/Engine/Plugins/Experimental/UAF/UAFTests/Source/UAFTestSuite/Public/AnimNextTest.h
--- a/Engine/Plugins/Experimental/UAF/UAFTests/Source/UAFTestSuite/Public/AnimNextTest.h
+++ b/Engine/Plugins/Experimental/UAF/UAFTests/Source/UAFTestSuite/Public/AnimNextTest.h
@@ -23,6 +23,7 @@ struct FUtils final
 {
 	// Clean up after tests. Clears transaction buffer, collects garbage
 	static UE_API void CleanupAfterTests();
+	static UE_API class UWorld* GetWorld();
 };
 
 }
diff --git a/Engine/Plugins/Experimental/UAF/UAFTests/Source/UAFTests/Private/UAFModuleTests.cpp b/Engine/Plugins/Experimental/UAF/UAFTests/Source/UAFTests/Private/UAFModuleTests.cpp
--- a/Engine/Plugins/Experimental/UAF/UAFTests/Source/UAFTests/Private/UAFModuleTests.cpp
+++ b/Engine/Plugins/Experimental/UAF/UAFTests/Source/UAFTests/Private/UAFModuleTests.cpp
@@ -8,6 +8,10 @@
 #include "Editor.h"
 #include "ScopedTransaction.h"
 #include "Subsystems/AssetEditorSubsystem.h"
+#include "AnimNextTest.h"
+#include "Module/SystemReference.h"
+#include "Module/RigUnit_AnimNextModuleEvents.h"
+#include "Module/AnimNextModuleInitMethod.h"
 #include "Module/AnimNextModule.h"
 #include "Module/AnimNextModuleFactory.h"
 #include "Graph/RigUnit_AnimNextRunAnimationGraph_v2.h"
@@ -294,5 +298,43 @@ protected:
 	}
 };
 
+namespace UE::UAF::Tests
+{
+
+IMPLEMENT_SIMPLE_AUTOMATION_TEST(FSystemReference_Creation, "Animation.UAF.Runtime.SystemReference.Creation", EAutomationTestFlags::EditorContext | EAutomationTestFlags::EngineFilter)
+
+bool FSystemReference_Creation::RunTest(const FString& InParameters)
+{
+	using namespace UE::UAF;
+
+	ON_SCOPE_EXIT{ FUtils::CleanupAfterTests(); };
+
+	UFactory* Factory = NewObject<UFactory>(GetTransientPackage(), UAnimNextModuleFactory::StaticClass());
+	UAnimNextModule* Asset = CastChecked<UAnimNextModule>(Factory->FactoryCreateNew(UAnimNextModule::StaticClass(), GetTransientPackage(), TEXT("TestAsset"), RF_Transient, nullptr, nullptr, NAME_None));
+	UE_RETURN_ON_ERROR(Asset != nullptr, "FSystemReference_Creation -> Failed to create asset");
+
+	TArray<FString> Messages;
+	FRigVMRuntimeSettings RuntimeSettings;
+	RuntimeSettings.SetLogFunction([&Messages](const FRigVMLogSettings& InLogSettings, const FRigVMExecuteContext* InContext, const FString& Message)
+	{
+		Messages.Add(Message);
+	});
+	Asset->GetRigVMExtendedExecuteContext().SetRuntimeSettings(RuntimeSettings);
+
+	// Needs a world context object, otherwise tick function creation fails.
+	FSystemReference SystemRef(Asset, FUtils::GetWorld(), EAnimNextModuleInitMethod::None);
+	UE_RETURN_ON_ERROR(SystemRef.IsValid(), "FSystemReference_Creation -> Failed to create System");
+
+	const FTickFunction* TickFunction = SystemRef.FindTickFunction(FRigUnit_AnimNextPrePhysicsEvent::DefaultEventName);
+	UE_RETURN_ON_ERROR(TickFunction != nullptr, "FSystemReference_Creation -> Failed to find tick functions.");
+
+	SystemRef.RunEvent(FRigUnit_AnimNextPrePhysicsEvent::DefaultEventName, 0.0f);
+	UE_RETURN_ON_ERROR(Messages.IsEmpty(), "FSystemReference_Creation -> Unexpected number of messages.");
+
+	return true;
+}
+
+}
+
 #endif // WITH_EDITOR
 #endif // WITH_DEV_AUTOMATION_TESTS
diff --git a/Engine/Plugins/Experimental/UAF/UAFTests/Source/UAFTests/UAFTests.Build.cs b/Engine/Plugins/Experimental/UAF/UAFTests/Source/UAFTests/UAFTests.Build.cs
--- a/Engine/Plugins/Experimental/UAF/UAFTests/Source/UAFTests/UAFTests.Build.cs
+++ b/Engine/Plugins/Experimental/UAF/UAFTests/Source/UAFTests/UAFTests.Build.cs
@@ -44,6 +44,7 @@ public class UAFTests : ModuleRules
 				"WorkspaceEditor",
 				"UAFEditor",
 				"UAFStateTree",
+				"UAFTestSuite",
 				"UAFStateTreeEditor",
 				"UAFAnimGraph",
 				"UAFAnimGraphEditor",
