Mass UAF: Fix for crash on PIE stop due to dangling debug render proxies
#Synced-CL 48498703

diff --git a/Engine/Plugins/Experimental/UAF/UAF/Source/UAF/Public/Module/SystemReference.h b/Engine/Plugins/Experimental/UAF/UAF/Source/UAF/Public/Module/SystemReference.h
--- a/Engine/Plugins/Experimental/UAF/UAF/Source/UAF/Public/Module/SystemReference.h
+++ b/Engine/Plugins/Experimental/UAF/UAF/Source/UAF/Public/Module/SystemReference.h
@@ -27,6 +27,7 @@ class UMassUAFProcessor;
 class UMassUAFInitializer;
 class UMeshEntityComponentImpl_Skinned;
 class UEntitySkeletonModifierStackImpl;
+class UMassUAFDestructor;
 
 namespace UE::UAF
 {
@@ -72,6 +73,7 @@ private:
 	friend FMassUAFFragment;
 	friend UMassUAFProcessor;
 	friend UMassUAFInitializer;
+	friend UMassUAFDestructor;
 
 	// Construct an owning handle to a system instance.
 	// Constructs a FAnimNextModuleInstance for use by worker threads.
diff --git a/Engine/Plugins/Experimental/UAF/UAFMass/Source/UAFMass/Private/Processors/MassUAFProcessor.cpp b/Engine/Plugins/Experimental/UAF/UAFMass/Source/UAFMass/Private/Processors/MassUAFProcessor.cpp
--- a/Engine/Plugins/Experimental/UAF/UAFMass/Source/UAFMass/Private/Processors/MassUAFProcessor.cpp
+++ b/Engine/Plugins/Experimental/UAF/UAFMass/Source/UAFMass/Private/Processors/MassUAFProcessor.cpp
@@ -138,4 +138,36 @@ void UMassAnimPoseDebugProcessor::Execute(FMassEntityManager& EntityManager, FMa
 		{
 		}
 	});*/
+}
+
+UMassUAFDestructor::UMassUAFDestructor()
+	: EntityQuery(*this)
+{
+	ObservedType = FMassUAFFragment::StaticStruct();
+	ObservedOperations = EMassObservedOperationFlags::Remove;
+	ExecutionFlags = (int32)EProcessorExecutionFlags::AllNetModes;
+	bRequiresGameThreadExecution = false;
+}
+
+void UMassUAFDestructor::ConfigureQueries(const TSharedRef<FMassEntityManager>& EntityManager)
+{
+	EntityQuery.AddRequirement<FMassUAFFragment>(EMassFragmentAccess::ReadWrite);
+}
+
+void UMassUAFDestructor::Execute(FMassEntityManager& EntityManager, FMassExecutionContext& Context)
+{
+	using namespace UE::UAF;
+	EntityQuery.ForEachEntityChunk(Context, [this](FMassExecutionContext& Context)
+		{
+			FTaskTagScope Scope(ETaskTag::EParallelGameThread);
+
+			const int32 NumEntities = Context.GetNumEntities();
+			TArrayView<FMassUAFFragment> UAFFragmentList = Context.GetMutableFragmentView<FMassUAFFragment>();
+
+			for (FMassExecutionContext::FEntityIterator EntityIt = Context.CreateEntityIterator(); EntityIt; ++EntityIt)
+			{
+				FMassUAFFragment& UAFFragment = UAFFragmentList[EntityIt];
+				UAFFragment.SystemReference.Reset();
+			}
+		});
 }
\ No newline at end of file
diff --git a/Engine/Plugins/Experimental/UAF/UAFMass/Source/UAFMass/Private/Processors/MassUAFProcessor.h b/Engine/Plugins/Experimental/UAF/UAFMass/Source/UAFMass/Private/Processors/MassUAFProcessor.h
--- a/Engine/Plugins/Experimental/UAF/UAFMass/Source/UAFMass/Private/Processors/MassUAFProcessor.h
+++ b/Engine/Plugins/Experimental/UAF/UAFMass/Source/UAFMass/Private/Processors/MassUAFProcessor.h
@@ -49,6 +49,21 @@ class UMassAnimPoseDebugProcessor : public UMassProcessor
 public:
 	UMassAnimPoseDebugProcessor();
 
+protected:
+	virtual void ConfigureQueries(const TSharedRef<FMassEntityManager>& EntityManager) override;
+	virtual void Execute(FMassEntityManager& EntityManager, FMassExecutionContext& Context) override;
+
+	FMassEntityQuery EntityQuery;
+};
+
+UCLASS()
+class UMassUAFDestructor : public UMassObserverProcessor
+{
+	GENERATED_BODY()
+
+public:
+	UMassUAFDestructor();
+
 protected:
 	virtual void ConfigureQueries(const TSharedRef<FMassEntityManager>& EntityManager) override;
 	virtual void Execute(FMassEntityManager& EntityManager, FMassExecutionContext& Context) override;
